Fix DDC build error

diff --git a/Templates/TP_HandheldARBP/Config/DefaultEngine.ini b/Templates/TP_HandheldARBP/Config/DefaultEngine.ini
--- a/Templates/TP_HandheldARBP/Config/DefaultEngine.ini
+++ b/Templates/TP_HandheldARBP/Config/DefaultEngine.ini
@@ -29,7 +29,7 @@ r.TranslucentSortPolicy=0
 TranslucentSortAxis=(X=0.000000,Y=-1.000000,Z=0.000000)
 r.CustomDepth=1
 r.CustomDepthTemporalAAJitter=True
-r.PostProcessing.PropagateAlpha=False
+r.PostProcessing.PropagateAlpha=0
 r.UsePreExposure=False
 r.DefaultFeature.Bloom=True
 r.DefaultFeature.AmbientOcclusion=False
