Fixes some of the projects with r.AntiAliasingMethod refactor

diff --git a/Templates/TP_HandheldARBP/Config/DefaultEngine.ini b/Templates/TP_HandheldARBP/Config/DefaultEngine.ini
--- a/Templates/TP_HandheldARBP/Config/DefaultEngine.ini
+++ b/Templates/TP_HandheldARBP/Config/DefaultEngine.ini
@@ -38,7 +38,7 @@ r.DefaultFeature.AutoExposure.Method=0
 r.DefaultFeature.MotionBlur=False
 r.DefaultFeature.LensFlare=False
 r.TemporalAA.Upsampling=False
-r.DefaultFeature.AntiAliasing=0
+r.AntiAliasingMethod=0
 r.DefaultFeature.PointLightUnits=1
 r.DefaultFeature.SpotLightUnits=1
 r.Shadow.UnbuiltPreviewInGame=True
diff --git a/Templates/TP_VirtualRealityBP/Config/DefaultEngine.ini b/Templates/TP_VirtualRealityBP/Config/DefaultEngine.ini
--- a/Templates/TP_VirtualRealityBP/Config/DefaultEngine.ini
+++ b/Templates/TP_VirtualRealityBP/Config/DefaultEngine.ini
@@ -53,7 +53,7 @@ r.DefaultFeature.MotionBlur=False
 r.DefaultFeature.LensFlare=False
 r.TemporalAA.Upsampling=False
 r.SSGI.Enable=False
-r.DefaultFeature.AntiAliasing=3
+r.AntiAliasingMethod=3
 r.DefaultFeature.LightUnits=1
 r.DefaultBackBufferPixelFormat=4
 r.Shadow.UnbuiltPreviewInGame=True
