ARCore - Use HLSLcc for external texture support

DXC does not support external textures in 5.0. ARCore requires external textures for support. We need to rollback to HLSLcc for the templates.
I did try using ModifyCompilationEnvironment, but the DXC flags aren't set by the time it is invoked.

This is TEMPORARY. Once external texture support lands in our DXC/ShaderConductor backend, we can revert this change to use DXC for ARCore.

diff --git a/Templates/TP_HandheldARBP/Config/DefaultEngine.ini b/Templates/TP_HandheldARBP/Config/DefaultEngine.ini
--- a/Templates/TP_HandheldARBP/Config/DefaultEngine.ini
+++ b/Templates/TP_HandheldARBP/Config/DefaultEngine.ini
@@ -77,6 +77,8 @@ r.SkinCache.SceneMemoryLimitInMB=128.000000
 r.GPUSkin.Limit2BoneInfluences=False
 r.SupportDepthOnlyIndexBuffers=True
 r.SupportReversedIndexBuffers=True
+; Temporary fallback to HLSLcc for ARCore until DXC supports external textures
+r.OpenGL.ForceDXC=0
 
 [/Script/EngineSettings.GameMapsSettings]
 EditorStartupMap=/Game/HandheldAR/Maps/HandheldARBlankMap.HandheldARBlankMap
