ARCore: Re-enable DXC + disable Vulkan

* DXC now supports external textures, ARCore's back on the menu!
* ARCore doesn't support Vulkan (yet), so as a short-term fix, we can disable that combo in our templates/samples.

diff --git a/Templates/TP_HandheldARBP/Config/DefaultEngine.ini b/Templates/TP_HandheldARBP/Config/DefaultEngine.ini
--- a/Templates/TP_HandheldARBP/Config/DefaultEngine.ini
+++ b/Templates/TP_HandheldARBP/Config/DefaultEngine.ini
@@ -77,8 +77,6 @@ r.SkinCache.SceneMemoryLimitInMB=128.000000
 r.GPUSkin.Limit2BoneInfluences=False
 r.SupportDepthOnlyIndexBuffers=True
 r.SupportReversedIndexBuffers=True
-; Temporary fallback to HLSLcc for ARCore until DXC supports external textures
-r.OpenGL.ForceDXC=0
 
 [/Script/EngineSettings.GameMapsSettings]
 EditorStartupMap=/Game/HandheldAR/Maps/HandheldARBlankMap.HandheldARBlankMap
@@ -120,6 +118,7 @@ GameScreenshotSaveDirectory=(Path="Screenshots/")
 Orientation=Sensor
 MinSDKVersion=24
 TargetSDKVersion=24
+bSupportsVulkan=false
 
 [/Script/OculusHMD.OculusHMDRuntimeSettings]
 bAutoEnabled=True
