Project Templates: Enabled CSM support on levels with Force No Precomputed Lighting on TP_Blank and TP_BlankBP templates as a means of supporting dynamic shadows on mobile.

Updated supported Template projects to default to Vulkan SM6 on Linux and Windows.

diff --git a/Templates/TP_Blank/Config/DefaultEngine.ini b/Templates/TP_Blank/Config/DefaultEngine.ini
--- a/Templates/TP_Blank/Config/DefaultEngine.ini
+++ b/Templates/TP_Blank/Config/DefaultEngine.ini
@@ -31,3 +31,10 @@ MinSampleRate=8000.000000
 CompressionQualityModifier=1.000000
 AutoStreamingThreshold=0.000000
 SoundCueCookQualityIndex=-1
+
+[/Script/Engine.RendererSettings]
+r.Mobile.EnableNoPrecomputedLightingCSMShader=True
+
+[/Script/LinuxTargetPlatform.LinuxTargetSettings]
+-TargetedRHIs=SF_VULKAN_SM5
++TargetedRHIs=SF_VULKAN_SM6
\ No newline at end of file
diff --git a/Templates/TP_BlankBP/Config/DefaultEngine.ini b/Templates/TP_BlankBP/Config/DefaultEngine.ini
--- a/Templates/TP_BlankBP/Config/DefaultEngine.ini
+++ b/Templates/TP_BlankBP/Config/DefaultEngine.ini
@@ -32,3 +32,9 @@ CompressionQualityModifier=1.000000
 AutoStreamingThreshold=0.000000
 SoundCueCookQualityIndex=-1
 
+[/Script/Engine.RendererSettings]
+r.Mobile.EnableNoPrecomputedLightingCSMShader=True
+
+[/Script/LinuxTargetPlatform.LinuxTargetSettings]
+-TargetedRHIs=SF_VULKAN_SM5
++TargetedRHIs=SF_VULKAN_SM6
\ No newline at end of file
diff --git a/Templates/TP_FirstPerson/Config/DefaultEngine.ini b/Templates/TP_FirstPerson/Config/DefaultEngine.ini
--- a/Templates/TP_FirstPerson/Config/DefaultEngine.ini
+++ b/Templates/TP_FirstPerson/Config/DefaultEngine.ini
@@ -53,17 +53,7 @@ CompressionQualityModifier=1.000000
 AutoStreamingThreshold=0.000000
 SoundCueCookQualityIndex=-1
 
-[/Script/AndroidFileServerEditor.AndroidFileServerRuntimeSettings]
-bEnablePlugin=True
-bAllowNetworkConnection=True
-SecurityToken=73D7EAB348128FC7E708A79C30473B10
-bIncludeInShipping=False
-bAllowExternalStartInShipping=False
-bCompileAFSProject=False
-bUseCompression=False
-bLogFiles=False
-bReportStats=False
-ConnectionType=USBOnly
-bUseManualIPAddress=False
-ManualIPAddress=
+[/Script/LinuxTargetPlatform.LinuxTargetSettings]
+-TargetedRHIs=SF_VULKAN_SM5
++TargetedRHIs=SF_VULKAN_SM6
 
diff --git a/Templates/TP_FirstPersonBP/Config/DefaultEngine.ini b/Templates/TP_FirstPersonBP/Config/DefaultEngine.ini
--- a/Templates/TP_FirstPersonBP/Config/DefaultEngine.ini
+++ b/Templates/TP_FirstPersonBP/Config/DefaultEngine.ini
@@ -29,3 +29,7 @@ DefaultGraphicsRHI=DefaultGraphicsRHI_DX12
 +D3D12TargetedShaderFormats=PCD3D_SM6
 -D3D11TargetedShaderFormats=PCD3D_SM5
 +D3D11TargetedShaderFormats=PCD3D_SM5
+
+[/Script/LinuxTargetPlatform.LinuxTargetSettings]
+-TargetedRHIs=SF_VULKAN_SM5
++TargetedRHIs=SF_VULKAN_SM6
\ No newline at end of file
diff --git a/Templates/TP_HandheldARBP/Config/DefaultEngine.ini b/Templates/TP_HandheldARBP/Config/DefaultEngine.ini
--- a/Templates/TP_HandheldARBP/Config/DefaultEngine.ini
+++ b/Templates/TP_HandheldARBP/Config/DefaultEngine.ini
@@ -1,4 +1,5 @@
 [URL]
+
 [/Script/Engine.RendererSettings]
 r.MobileHDR=True
 r.Mobile.DisableVertexFog=True
@@ -170,3 +171,29 @@ DefaultGraphicsRHI=DefaultGraphicsRHI_DX12
 +D3D12TargetedShaderFormats=PCD3D_SM6
 -D3D11TargetedShaderFormats=PCD3D_SM5
 +D3D11TargetedShaderFormats=PCD3D_SM5
+Compiler=Default
+AudioSampleRate=48000
+AudioCallbackBufferFrameSize=1024
+AudioNumBuffersToEnqueue=1
+AudioMaxChannels=0
+AudioNumSourceWorkers=4
+SpatializationPlugin=
+SourceDataOverridePlugin=
+ReverbPlugin=
+OcclusionPlugin=
+CompressionOverrides=(bOverrideCompressionTimes=False,DurationThreshold=5.000000,MaxNumRandomBranches=0,SoundCueQualityIndex=0)
+CacheSizeKB=65536
+MaxChunkSizeOverrideKB=0
+bResampleForDevice=False
+MaxSampleRate=48000.000000
+HighSampleRate=32000.000000
+MedSampleRate=24000.000000
+LowSampleRate=12000.000000
+MinSampleRate=8000.000000
+CompressionQualityModifier=1.000000
+AutoStreamingThreshold=0.000000
+SoundCueCookQualityIndex=-1
+
+[/Script/LinuxTargetPlatform.LinuxTargetSettings]
+-TargetedRHIs=SF_VULKAN_SM5
++TargetedRHIs=SF_VULKAN_SM6
\ No newline at end of file
diff --git a/Templates/TP_ThirdPerson/Config/DefaultEngine.ini b/Templates/TP_ThirdPerson/Config/DefaultEngine.ini
--- a/Templates/TP_ThirdPerson/Config/DefaultEngine.ini
+++ b/Templates/TP_ThirdPerson/Config/DefaultEngine.ini
@@ -40,3 +40,7 @@ MinSampleRate=8000.000000
 CompressionQualityModifier=1.000000
 AutoStreamingThreshold=0.000000
 SoundCueCookQualityIndex=-1
+
+[/Script/LinuxTargetPlatform.LinuxTargetSettings]
+-TargetedRHIs=SF_VULKAN_SM5
++TargetedRHIs=SF_VULKAN_SM6
\ No newline at end of file
diff --git a/Templates/TP_ThirdPersonBP/Config/DefaultEngine.ini b/Templates/TP_ThirdPersonBP/Config/DefaultEngine.ini
--- a/Templates/TP_ThirdPersonBP/Config/DefaultEngine.ini
+++ b/Templates/TP_ThirdPersonBP/Config/DefaultEngine.ini
@@ -48,3 +48,7 @@ MinSampleRate=8000.000000
 CompressionQualityModifier=1.000000
 AutoStreamingThreshold=0.000000
 SoundCueCookQualityIndex=-1
+
+[/Script/LinuxTargetPlatform.LinuxTargetSettings]
+-TargetedRHIs=SF_VULKAN_SM5
++TargetedRHIs=SF_VULKAN_SM6
\ No newline at end of file
diff --git a/Templates/TP_TopDown/Config/DefaultEngine.ini b/Templates/TP_TopDown/Config/DefaultEngine.ini
--- a/Templates/TP_TopDown/Config/DefaultEngine.ini
+++ b/Templates/TP_TopDown/Config/DefaultEngine.ini
@@ -36,3 +36,6 @@ CompressionQualityModifier=1.000000
 AutoStreamingThreshold=0.000000
 SoundCueCookQualityIndex=-1
 
+[/Script/LinuxTargetPlatform.LinuxTargetSettings]
+-TargetedRHIs=SF_VULKAN_SM5
++TargetedRHIs=SF_VULKAN_SM6
diff --git a/Templates/TP_TopDownBP/Config/DefaultEngine.ini b/Templates/TP_TopDownBP/Config/DefaultEngine.ini
--- a/Templates/TP_TopDownBP/Config/DefaultEngine.ini
+++ b/Templates/TP_TopDownBP/Config/DefaultEngine.ini
@@ -86,3 +86,7 @@ CompressionQualityModifier=1.000000
 AutoStreamingThreshold=0.000000
 SoundCueCookQualityIndex=-1
 
+[/Script/LinuxTargetPlatform.LinuxTargetSettings]
+-TargetedRHIs=SF_VULKAN_SM5
++TargetedRHIs=SF_VULKAN_SM6
+
diff --git a/Templates/TP_VehicleAdv/Config/DefaultEngine.ini b/Templates/TP_VehicleAdv/Config/DefaultEngine.ini
--- a/Templates/TP_VehicleAdv/Config/DefaultEngine.ini
+++ b/Templates/TP_VehicleAdv/Config/DefaultEngine.ini
@@ -27,3 +27,29 @@ DefaultGraphicsRHI=DefaultGraphicsRHI_DX12
 +D3D12TargetedShaderFormats=PCD3D_SM6
 -D3D11TargetedShaderFormats=PCD3D_SM5
 +D3D11TargetedShaderFormats=PCD3D_SM5
+Compiler=Default
+AudioSampleRate=48000
+AudioCallbackBufferFrameSize=1024
+AudioNumBuffersToEnqueue=1
+AudioMaxChannels=0
+AudioNumSourceWorkers=4
+SpatializationPlugin=
+SourceDataOverridePlugin=
+ReverbPlugin=
+OcclusionPlugin=
+CompressionOverrides=(bOverrideCompressionTimes=False,DurationThreshold=5.000000,MaxNumRandomBranches=0,SoundCueQualityIndex=0)
+CacheSizeKB=65536
+MaxChunkSizeOverrideKB=0
+bResampleForDevice=False
+MaxSampleRate=48000.000000
+HighSampleRate=32000.000000
+MedSampleRate=24000.000000
+LowSampleRate=12000.000000
+MinSampleRate=8000.000000
+CompressionQualityModifier=1.000000
+AutoStreamingThreshold=0.000000
+SoundCueCookQualityIndex=-1
+
+[/Script/LinuxTargetPlatform.LinuxTargetSettings]
+-TargetedRHIs=SF_VULKAN_SM5
++TargetedRHIs=SF_VULKAN_SM6
\ No newline at end of file
diff --git a/Templates/TP_VehicleAdvBP/Config/DefaultEngine.ini b/Templates/TP_VehicleAdvBP/Config/DefaultEngine.ini
--- a/Templates/TP_VehicleAdvBP/Config/DefaultEngine.ini
+++ b/Templates/TP_VehicleAdvBP/Config/DefaultEngine.ini
@@ -55,3 +55,6 @@ CompressionQualityModifier=1.000000
 AutoStreamingThreshold=0.000000
 SoundCueCookQualityIndex=-1
 
+[/Script/LinuxTargetPlatform.LinuxTargetSettings]
+-TargetedRHIs=SF_VULKAN_SM5
++TargetedRHIs=SF_VULKAN_SM6
\ No newline at end of file
diff --git a/Templates/TP_VirtualRealityBP/Config/DefaultEngine.ini b/Templates/TP_VirtualRealityBP/Config/DefaultEngine.ini
--- a/Templates/TP_VirtualRealityBP/Config/DefaultEngine.ini
+++ b/Templates/TP_VirtualRealityBP/Config/DefaultEngine.ini
@@ -68,4 +68,8 @@ AutoStreamingThreshold=0.000000
 SoundCueCookQualityIndex=-1
 
 [/Script/NavigationSystem.NavigationSystemV1]
-bAllowClientSideNavigation=True
\ No newline at end of file
+bAllowClientSideNavigation=True
+
+[/Script/LinuxTargetPlatform.LinuxTargetSettings]
+-TargetedRHIs=SF_VULKAN_SM5
++TargetedRHIs=SF_VULKAN_SM6
