Disable RayTracing for HandheldAR template
#tests Creating new game from HandheldAR template
#rnx

diff --git a/Templates/TP_HandheldARBP/Config/DefaultEngine.ini b/Templates/TP_HandheldARBP/Config/DefaultEngine.ini
--- a/Templates/TP_HandheldARBP/Config/DefaultEngine.ini
+++ b/Templates/TP_HandheldARBP/Config/DefaultEngine.ini
@@ -69,6 +69,7 @@ r.SupportLowQualityLightmaps=True
 r.SupportPointLightWholeSceneShadows=True
 r.SupportAtmosphericFog=True
 r.SkinCache.CompileShaders=False
+r.RayTracing=False
 r.Mobile.EnableStaticAndCSMShadowReceivers=True
 r.Mobile.EnableMovableLightCSMShaderCulling=True
 r.Mobile.AllowDistanceFieldShadows=True
