HHAR Template | Disabling auto exposure

diff --git a/Templates/TP_HandheldARBP/Config/DefaultEngine.ini b/Templates/TP_HandheldARBP/Config/DefaultEngine.ini
--- a/Templates/TP_HandheldARBP/Config/DefaultEngine.ini
+++ b/Templates/TP_HandheldARBP/Config/DefaultEngine.ini
@@ -14,4 +14,8 @@ GlobalDefaultServerGameMode=None
 
 [/Script/AndroidRuntimeSettings.AndroidRuntimeSettings]
 Orientation=Sensor
-bBuildForES31=False
\ No newline at end of file
+bBuildForES31=False
+
+[/Script/Engine.RendererSettings]
+r.DefaultFeature.AutoExposure=False
+
